Fix transaction activity details missing on server

diff --git a/Engine/Plugins/Developer/Concert/ConcertApp/MultiUserServer/Source/MultiUserServer/Private/Widgets/SessionTabs/ServerSessionHistoryControllerBase.cpp b/Engine/Plugins/Developer/Concert/ConcertApp/MultiUserServer/Source/MultiUserServer/Private/Widgets/SessionTabs/ServerSessionHistoryControllerBase.cpp
--- a/Engine/Plugins/Developer/Concert/ConcertApp/MultiUserServer/Source/MultiUserServer/Private/Widgets/SessionTabs/ServerSessionHistoryControllerBase.cpp
+++ b/Engine/Plugins/Developer/Concert/ConcertApp/MultiUserServer/Source/MultiUserServer/Private/Widgets/SessionTabs/ServerSessionHistoryControllerBase.cpp
@@ -5,9 +5,10 @@
 #include "ConcertSyncSessionDatabase.h"
 #include "IConcertSession.h"
 #include "IConcertSyncServer.h"
+#include "ServerUndoHistoryReflectionDataProvider.h"
 
 FServerSessionHistoryControllerBase::FServerSessionHistoryControllerBase(FGuid SessionId, SSessionHistory::FArguments Arguments)
-	: FAbstractSessionHistoryController(MoveTemp(Arguments))
+	: FAbstractSessionHistoryController(MoveTemp(Arguments.UndoHistoryReflectionProvider(MakeShared<UE::MultiUserServer::FServerUndoHistoryReflectionDataProvider>())))
 	, SessionId(MoveTemp(SessionId))
 {}
 
diff --git a/Engine/Plugins/Developer/Concert/ConcertApp/MultiUserServer/Source/MultiUserServer/Private/Widgets/SessionTabs/ServerUndoHistoryReflectionDataProvider.cpp b/Engine/Plugins/Developer/Concert/ConcertApp/MultiUserServer/Source/MultiUserServer/Private/Widgets/SessionTabs/ServerUndoHistoryReflectionDataProvider.cpp
new file mode 100644
--- /dev/null
+++ b/Engine/Plugins/Developer/Concert/ConcertApp/MultiUserServer/Source/MultiUserServer/Private/Widgets/SessionTabs/ServerUndoHistoryReflectionDataProvider.cpp
@@ -0,0 +1,30 @@
+// Copyright Epic Games, Inc. All Rights Reserved.
+
+#include "ServerUndoHistoryReflectionDataProvider.h"
+
+#include "Algo/AnyOf.h"
+
+void UE::MultiUserServer::FServerUndoHistoryReflectionDataProvider::SetTransactionContext(const FConcertTransactionEventBase& InTransaction)
+{
+	TransactionContext = InTransaction;
+}
+
+bool UE::MultiUserServer::FServerUndoHistoryReflectionDataProvider::HasClassDisplayName(const FSoftClassPath& ClassPath) const
+{
+	return Algo::AnyOf(TransactionContext.ExportedObjects, [ClassPath](const FConcertExportedObject& Object)
+	{
+		return Object.ObjectId.ObjectClassPathName.ToString() == ClassPath.ToString();
+	});
+}
+
+TOptional<FString> UE::MultiUserServer::FServerUndoHistoryReflectionDataProvider::GetClassDisplayName(const FSoftClassPath& ClassPath) const
+{
+	const FString AsString = ClassPath.ToString();
+	int32 Index;
+	if (AsString.FindLastChar(TEXT('/'), Index))
+	{
+		return AsString.RightChop(Index + 1);
+	}
+
+	return {};
+}
diff --git a/Engine/Plugins/Developer/Concert/ConcertApp/MultiUserServer/Source/MultiUserServer/Private/Widgets/SessionTabs/ServerUndoHistoryReflectionDataProvider.h b/Engine/Plugins/Developer/Concert/ConcertApp/MultiUserServer/Source/MultiUserServer/Private/Widgets/SessionTabs/ServerUndoHistoryReflectionDataProvider.h
new file mode 100644
--- /dev/null
+++ b/Engine/Plugins/Developer/Concert/ConcertApp/MultiUserServer/Source/MultiUserServer/Private/Widgets/SessionTabs/ServerUndoHistoryReflectionDataProvider.h
@@ -0,0 +1,32 @@
+// Copyright Epic Games, Inc. All Rights Reserved.
+
+#pragma once
+
+#include "CoreMinimal.h"
+#include "ConcertTransactionEvents.h"
+#include "Session/Activity/IConcertReflectionDataProvider.h"
+
+namespace UE::MultiUserServer
+{
+	class FServerUndoHistoryReflectionDataProvider : public ConcertSharedSlate::IConcertReflectionDataProvider
+	{
+	public:
+
+		//~ Begin IConcertReflectionDataProvider Interface
+		virtual void SetTransactionContext(const FConcertTransactionEventBase& InTransaction) override;
+		//~ End IConcertReflectionDataProvider Interface
+
+		//~ Begin IReflectionDataProvider Interface
+		virtual bool HasClassDisplayName(const FSoftClassPath& ClassPath) const override;
+		virtual TOptional<FString> GetClassDisplayName(const FSoftClassPath& ClassPath) const override;
+		virtual bool SupportsGetPropertyReflectionData() const override { return false; }
+		virtual TOptional<UndoHistory::FPropertyReflectionData> GetPropertyReflectionData(const FSoftClassPath& ClassPath, FName PropertyName) const override { return {}; }
+		//~ End IReflectionDataProvider Interface
+
+	private:
+
+		FConcertTransactionEventBase TransactionContext;
+	};
+}
+
+
diff --git a/Engine/Plugins/Developer/Concert/ConcertUI/ConcertSharedSlate/Source/ConcertSharedSlate/ConcertSharedSlate.Build.cs b/Engine/Plugins/Developer/Concert/ConcertUI/ConcertSharedSlate/Source/ConcertSharedSlate/ConcertSharedSlate.Build.cs
--- a/Engine/Plugins/Developer/Concert/ConcertUI/ConcertSharedSlate/Source/ConcertSharedSlate/ConcertSharedSlate.Build.cs
+++ b/Engine/Plugins/Developer/Concert/ConcertUI/ConcertSharedSlate/Source/ConcertSharedSlate/ConcertSharedSlate.Build.cs
@@ -13,6 +13,7 @@ namespace UnrealBuildTool.Rules
 					"Core",
 					"CoreUObject",
 					"Slate",
+					"UndoHistory",
 				}
 			);
 
diff --git a/Engine/Plugins/Developer/Concert/ConcertUI/ConcertSharedSlate/Source/ConcertSharedSlate/Private/Session/Activity/SConcertSessionActivities.cpp b/Engine/Plugins/Developer/Concert/ConcertUI/ConcertSharedSlate/Source/ConcertSharedSlate/Private/Session/Activity/SConcertSessionActivities.cpp
--- a/Engine/Plugins/Developer/Concert/ConcertUI/ConcertSharedSlate/Source/ConcertSharedSlate/Private/Session/Activity/SConcertSessionActivities.cpp
+++ b/Engine/Plugins/Developer/Concert/ConcertUI/ConcertSharedSlate/Source/ConcertSharedSlate/Private/Session/Activity/SConcertSessionActivities.cpp
@@ -2,6 +2,15 @@
 
 #include "Session/Activity/SConcertSessionActivities.h"
 
+#include "ConcertFrontendStyle.h"
+#include "ConcertFrontendUtils.h"
+#include "ConcertLogGlobal.h"
+#include "ConcertTransactionEvents.h"
+#include "ConcertWorkspaceData.h"
+#include "Session/Activity/IConcertReflectionDataProvider.h"
+#include "SessionActivityUtils.h"
+#include "SPackageDetails.h"
+
 #include "Algo/Transform.h"
 #include "Misc/AsyncTaskNotification.h"
 #include "Misc/TransactionObjectEvent.h"
@@ -22,13 +31,6 @@
 #include "Widgets/Text/STextBlock.h"
 #include "Widgets/Views/STableRow.h"
 #include "Styling/AppStyle.h"
-#include "ConcertFrontendStyle.h"
-#include "ConcertFrontendUtils.h"
-#include "ConcertLogGlobal.h"
-#include "ConcertTransactionEvents.h"
-#include "ConcertWorkspaceData.h"
-#include "SessionActivityUtils.h"
-#include "SPackageDetails.h"
 
 #include "Session/Activity/PredefinedActivityColumns.h"
 
@@ -238,6 +240,13 @@ void SConcertSessionActivities::Construct(const FArguments& InArgs)
 	SearchTextFilter = MakeShared<TTextFilter<const FConcertSessionActivity&>>(TTextFilter<const FConcertSessionActivity&>::FItemToStringArray::CreateSP(this, &SConcertSessionActivities::PopulateSearchStrings));
 	SearchTextFilter->OnChanged().AddSP(this, &SConcertSessionActivities::OnActivityFilterUpdated);
 
+	if (InArgs._UndoHistoryReflectionProvider)
+	{
+		ConcertReflectionDataProvider = InArgs._UndoHistoryReflectionProvider;
+	}
+	const TSharedRef<UE::UndoHistory::IReflectionDataProvider> ReflectionDataProvider =
+		ConcertReflectionDataProvider ? ConcertReflectionDataProvider->ToSharedRef() : UE::UndoHistory::CreateDefaultReflectionProvider();
+	
 	ActiveFilterFlags = QueryActiveActivityFilters();
 	ChildSlot
 	[
@@ -311,7 +320,7 @@ void SConcertSessionActivities::Construct(const FArguments& InArgs)
 
 					+SScrollBox::Slot()
 					[
-						SAssignNew(TransactionDetailsPanel, SUndoHistoryDetails)
+						SAssignNew(TransactionDetailsPanel, SUndoHistoryDetails, ReflectionDataProvider)
 						.Visibility(EVisibility::Collapsed)
 					]
 
@@ -835,6 +844,10 @@ void SConcertSessionActivities::DisplayTransactionDetails(const FConcertSessionA
 		TransactionDiff.DiffMap.Emplace(TransactionObjectId.ObjectPathName, MoveTemp(Event));
 	}
 
+	if (ConcertReflectionDataProvider)
+	{
+		ConcertReflectionDataProvider.GetValue()->SetTransactionContext(InTransaction);
+	}
 	TransactionDetailsPanel->SetSelectedTransaction(MoveTemp(TransactionDiff));
 	SetDetailsPanelVisibility(TransactionDetailsPanel.Get());
 }
diff --git a/Engine/Plugins/Developer/Concert/ConcertUI/ConcertSharedSlate/Source/ConcertSharedSlate/Private/Session/History/SSessionHistory.cpp b/Engine/Plugins/Developer/Concert/ConcertUI/ConcertSharedSlate/Source/ConcertSharedSlate/Private/Session/History/SSessionHistory.cpp
--- a/Engine/Plugins/Developer/Concert/ConcertUI/ConcertSharedSlate/Source/ConcertSharedSlate/Private/Session/History/SSessionHistory.cpp
+++ b/Engine/Plugins/Developer/Concert/ConcertUI/ConcertSharedSlate/Source/ConcertSharedSlate/Private/Session/History/SSessionHistory.cpp
@@ -74,6 +74,7 @@ void SSessionHistory::Construct(const FArguments& InArgs)
 		.ColumnVisibilitySnapshot(InArgs._ColumnVisibilitySnapshot)
 		.SaveColumnVisibilitySnapshot(InArgs._SaveColumnVisibilitySnapshot)
 		.SelectionMode(InArgs._SelectionMode)
+		.UndoHistoryReflectionProvider(InArgs._UndoHistoryReflectionProvider)
 		.DarkenMutedActivities(InArgs._DarkenMutedActivities);
 
 	ChildSlot
diff --git a/Engine/Plugins/Developer/Concert/ConcertUI/ConcertSharedSlate/Source/ConcertSharedSlate/Public/Session/Activity/IConcertReflectionDataProvider.h b/Engine/Plugins/Developer/Concert/ConcertUI/ConcertSharedSlate/Source/ConcertSharedSlate/Public/Session/Activity/IConcertReflectionDataProvider.h
new file mode 100644
--- /dev/null
+++ b/Engine/Plugins/Developer/Concert/ConcertUI/ConcertSharedSlate/Source/ConcertSharedSlate/Public/Session/Activity/IConcertReflectionDataProvider.h
@@ -0,0 +1,20 @@
+// Copyright Epic Games, Inc. All Rights Reserved.
+
+#pragma once
+
+#include "CoreMinimal.h"
+#include "Widgets/IReflectionDataProvider.h"
+
+struct FConcertTransactionEventBase;
+
+namespace UE::ConcertSharedSlate
+{
+	/** Passed on to SUndoHistoryDetails to display reflection data in non-editor builds. Used e.g. by the server. */
+	class CONCERTSHAREDSLATE_API IConcertReflectionDataProvider : public UndoHistory::IReflectionDataProvider 
+	{
+	public:
+
+		/** Called just before SUndoHistoryDetails::SetSelectedTransaction. Injects the related property data saved in the activity. */
+		virtual void SetTransactionContext(const FConcertTransactionEventBase& InTransaction) = 0;
+	};
+}
diff --git a/Engine/Plugins/Developer/Concert/ConcertUI/ConcertSharedSlate/Source/ConcertSharedSlate/Public/Session/Activity/SConcertSessionActivities.h b/Engine/Plugins/Developer/Concert/ConcertUI/ConcertSharedSlate/Source/ConcertSharedSlate/Public/Session/Activity/SConcertSessionActivities.h
--- a/Engine/Plugins/Developer/Concert/ConcertUI/ConcertSharedSlate/Source/ConcertSharedSlate/Public/Session/Activity/SConcertSessionActivities.h
+++ b/Engine/Plugins/Developer/Concert/ConcertUI/ConcertSharedSlate/Source/ConcertSharedSlate/Public/Session/Activity/SConcertSessionActivities.h
@@ -25,6 +25,11 @@ class SPackageDetails;
 class SUndoHistoryDetails;
 class STextBlock;
 
+namespace UE::ConcertSharedSlate
+{
+	class IConcertReflectionDataProvider;
+}
+
 /** Filters for the concert session activity view. */
 enum class EConcertActivityFilterFlags
 {
@@ -145,6 +150,9 @@ public:
 		/** Optional. How the activities may be selected */
 		SLATE_ARGUMENT(TOptional<ESelectionMode::Type>, SelectionMode)
 
+		/** Optional. Used by non-editor instances to provide reflection data about transaction activities. */
+		SLATE_ARGUMENT(TSharedPtr<UE::ConcertSharedSlate::IConcertReflectionDataProvider>, UndoHistoryReflectionProvider)
+
 		/** Optional. Whether to reduce focus to activities by darkening when the activity is muted (default: true). */
 		SLATE_ARGUMENT(bool, DarkenMutedActivities)
 	SLATE_END_ARGS();
@@ -247,6 +255,9 @@ private:
 	/** The list view widget displaying the activities. */
 	TSharedPtr<SListView<TSharedPtr<FConcertSessionActivity>>> ActivityView;
 
+	/** Provides reflection data for SUndoHistoryDetails for non-editor builds. */
+	TOptional<TSharedPtr<UE::ConcertSharedSlate::IConcertReflectionDataProvider>> ConcertReflectionDataProvider; 
+
 	/** Used to overlay a widget over a column widget (add an extra layer above the normal one) */
 	FMakeColumnOverlayWidgetFunc MakeColumnOverlayWidgetFn;
 
diff --git a/Engine/Plugins/Developer/Concert/ConcertUI/ConcertSharedSlate/Source/ConcertSharedSlate/Public/Session/History/SSessionHistory.h b/Engine/Plugins/Developer/Concert/ConcertUI/ConcertSharedSlate/Source/ConcertSharedSlate/Public/Session/History/SSessionHistory.h
--- a/Engine/Plugins/Developer/Concert/ConcertUI/ConcertSharedSlate/Source/ConcertSharedSlate/Public/Session/History/SSessionHistory.h
+++ b/Engine/Plugins/Developer/Concert/ConcertUI/ConcertSharedSlate/Source/ConcertSharedSlate/Public/Session/History/SSessionHistory.h
@@ -8,6 +8,11 @@
 #include "Widgets/DeclarativeSyntaxSupport.h"
 #include "Widgets/SCompoundWidget.h"
 
+namespace UE::ConcertSharedSlate
+{
+	class IConcertReflectionDataProvider;
+}
+
 class SSearchBox;
 class FConcertSessionActivitiesOptions;
 
@@ -58,7 +63,10 @@ public:
 		
 		/** Optional. How the activities may be selected. Default to Single. */
 		SLATE_ARGUMENT(ESelectionMode::Type, SelectionMode)
-
+	
+		/** Optional. Used by non-editor instances to provide reflection data about transaction activities. */
+		SLATE_ARGUMENT(TSharedPtr<UE::ConcertSharedSlate::IConcertReflectionDataProvider>, UndoHistoryReflectionProvider)
+	
 		/** Optional. Whether to reduce focus to activities by darkening when the activity is muted (default: true). */
 		SLATE_ARGUMENT(bool, DarkenMutedActivities)
 	
diff --git a/Engine/Source/Developer/UndoHistory/Private/Widgets/ClassLookupReflectionDataProvider.cpp b/Engine/Source/Developer/UndoHistory/Private/Widgets/ClassLookupReflectionDataProvider.cpp
new file mode 100644
--- /dev/null
+++ b/Engine/Source/Developer/UndoHistory/Private/Widgets/ClassLookupReflectionDataProvider.cpp
@@ -0,0 +1,79 @@
+// Copyright Epic Games, Inc. All Rights Reserved.
+
+#include "ClassLookupReflectionDataProvider.h"
+
+#include "UObject/Class.h"
+#include "UObject/EnumProperty.h"
+#include "UObject/UnrealType.h"
+
+namespace UE::UndoHistory
+{
+	bool FClassLookupReflectionDataProvider::HasClassDisplayName(const FSoftClassPath& ClassPath) const
+	{
+		return LookUpClass(ClassPath).IsSet();
+	}
+
+	TOptional<FString> FClassLookupReflectionDataProvider::GetClassDisplayName(const FSoftClassPath& ClassPath) const
+	{
+		const TOptional<TNonNullPtr<UClass>> Class = LookUpClass(ClassPath);
+		return Class ? Class->GetName() : TOptional<FString>{};
+	}
+
+	TOptional<FPropertyReflectionData> FClassLookupReflectionDataProvider::GetPropertyReflectionData(const FSoftClassPath& ClassPath, FName PropertyName) const
+	{
+		if (const TOptional<TNonNullPtr<FProperty>> Property = LookUpProperty(ClassPath, PropertyName))
+		{
+			FPropertyReflectionData Result{}; // Value init so if new members are added they have defined values
+			Result.PropertyType = GetPropertyType(*Property);
+			Property->GetCPPMacroType(Result.CppMacroType);
+			Result.TypeName = Property->GetClass()->GetName();
+			Result.PropertyFlags = Property->PropertyFlags;
+			return Result;
+		}
+		return {};
+	}
+
+	TOptional<TNonNullPtr<UClass>> FClassLookupReflectionDataProvider::LookUpClass(const FSoftClassPath& ClassPath) const
+	{
+		UClass* Class = LoadObject<UClass>(nullptr, *ClassPath.ToString());
+		return LIKELY(Class) ? TOptional<TNonNullPtr<UClass>>{ Class } : TOptional<TNonNullPtr<UClass>>{};
+	}
+
+	TOptional<TNonNullPtr<FProperty>> FClassLookupReflectionDataProvider::LookUpProperty(const FSoftClassPath& ClassPath, FName PropertyName) const
+	{
+		if (const TOptional<TNonNullPtr<UClass>> Class = LookUpClass(ClassPath))
+		{
+			FProperty* Property = Class->FindPropertyByName(PropertyName);
+			return LIKELY(Property) ? TOptional<TNonNullPtr<FProperty>>{ Property } : TOptional<TNonNullPtr<FProperty>>{};
+		}
+		return {};
+	}
+
+	TSharedRef<IReflectionDataProvider> CreateDefaultReflectionProvider()
+	{
+		return MakeShared<FClassLookupReflectionDataProvider>();
+	}
+
+	EPropertyType FClassLookupReflectionDataProvider::GetPropertyType(const FProperty& Property) const
+	{
+		if (CastField<FObjectPropertyBase>(&Property))
+		{
+			return EPropertyType::ObjectProperty;
+		}
+		if (CastField<FStructProperty>(&Property))
+		{
+			return EPropertyType::StructProperty;
+		}
+		if (CastField<FEnumProperty>(&Property))
+		{
+			return EPropertyType::EnumProperty;
+		}
+		if (CastField<FArrayProperty>(&Property))
+		{
+			return EPropertyType::StructProperty;
+		}
+		return EPropertyType::Other;
+	}
+}
+
+
diff --git a/Engine/Source/Developer/UndoHistory/Private/Widgets/ClassLookupReflectionDataProvider.h b/Engine/Source/Developer/UndoHistory/Private/Widgets/ClassLookupReflectionDataProvider.h
new file mode 100644
--- /dev/null
+++ b/Engine/Source/Developer/UndoHistory/Private/Widgets/ClassLookupReflectionDataProvider.h
@@ -0,0 +1,31 @@
+// Copyright Epic Games, Inc. All Rights Reserved.
+
+#pragma once
+
+#include "CoreMinimal.h"
+#include "Templates/NonNullPointer.h"
+#include "Widgets/IReflectionDataProvider.h"
+
+namespace UE::UndoHistory
+{
+	/** Just uses LoadObject to look up class data. Works on editor builds but not on most programs, like MultiUserServer. */
+	class FClassLookupReflectionDataProvider : public IReflectionDataProvider
+	{
+	public:
+
+		//~ Begin IReflectionDataProvider Interface
+		virtual bool HasClassDisplayName(const FSoftClassPath& ClassPath) const override;
+		virtual TOptional<FString> GetClassDisplayName(const FSoftClassPath& ClassPath) const override;
+		virtual bool SupportsGetPropertyReflectionData() const override { return true; }
+		virtual TOptional<FPropertyReflectionData> GetPropertyReflectionData(const FSoftClassPath& ClassPath, FName PropertyName) const override;
+		//~ End IReflectionDataProvider Interface
+
+	private:
+
+		TOptional<TNonNullPtr<UClass>> LookUpClass(const FSoftClassPath& ClassPath) const;
+		TOptional<TNonNullPtr<FProperty>> LookUpProperty(const FSoftClassPath& ClassPath, FName PropertyName) const;
+		EPropertyType GetPropertyType(const FProperty& Property) const;
+	};
+}
+
+
diff --git a/Engine/Source/Developer/UndoHistory/Private/Widgets/SUndoHistoryDetails.cpp b/Engine/Source/Developer/UndoHistory/Private/Widgets/SUndoHistoryDetails.cpp
--- a/Engine/Source/Developer/UndoHistory/Private/Widgets/SUndoHistoryDetails.cpp
+++ b/Engine/Source/Developer/UndoHistory/Private/Widgets/SUndoHistoryDetails.cpp
@@ -85,7 +85,7 @@ public:
 	{ }
 		SLATE_ARGUMENT(TSharedPtr<FTransactionObjectEvent>, TransactionEvent)
 		SLATE_ARGUMENT(FString, Name)
-		SLATE_ARGUMENT(FString, Type)
+		SLATE_ARGUMENT(TOptional<FString>, Type)
 		SLATE_ATTRIBUTE(FText, FilterText)
 	SLATE_END_ARGS()
 
@@ -131,10 +131,10 @@ public:
 							.HighlightText(FilterText)
 					];
 		}
-		else if (ColumnName == TransactionDetailsUI::TypeLabel)
+		else if (ColumnName == TransactionDetailsUI::TypeLabel && ensure(Type))
 		{
 			return SNew(STextBlock)
-				.Text(FText::FromString(Type));
+				.Text(FText::FromString(*Type));
 		}
 		else if (ColumnName == TransactionDetailsUI::ChangeFlagsLabel)
 		{
@@ -213,17 +213,14 @@ private:
 
 	TSharedPtr<FTransactionObjectEvent> TransactionEvent;
 	FString Name;
-	FString Type;
+	TOptional<FString> Type;
 	TAttribute<FText> FilterText;
 
 };
 
-void SUndoHistoryDetails::Construct(const FArguments& InArgs)
+void SUndoHistoryDetails::Construct(const FArguments& InArgs, TSharedRef<UE::UndoHistory::IReflectionDataProvider> InReflectionData)
 {
-	const auto OnGetChildren = [](FUndoDetailsTreeNodePtr InNode, TArray<FUndoDetailsTreeNodePtr> & OutChildren)
-	{
-		OutChildren = InNode->Children;
-	};
+	ReflectionData = MoveTemp(InReflectionData);
 
 	auto FilterDelegate = FTreeItemTextFilter::FItemToStringArray::CreateSP(this, &SUndoHistoryDetails::PopulateSearchStrings);
 	SearchBoxFilter = MakeShared<FTreeItemTextFilter>(FilterDelegate);
@@ -231,7 +228,11 @@ void SUndoHistoryDetails::Construct(const FArguments& InArgs)
 	
 	bNeedsRefresh = true;
 	bNeedsExpansion = false;
-
+	
+	const auto OnGetChildren = [](FUndoDetailsTreeNodePtr InNode, TArray<FUndoDetailsTreeNodePtr> & OutChildren)
+	{
+		OutChildren = InNode->Children;
+	};
 	ChildSlot
 		[
 			SNew(SBorder)
@@ -320,6 +321,7 @@ void SUndoHistoryDetails::Construct(const FArguments& InArgs)
 									.DefaultLabel(LOCTEXT("NameColumnHeaderName", "Modified objects and properties"))
 									+ SHeaderRow::Column(TransactionDetailsUI::TypeLabel)
 									.FillWidth(20.0f)
+									.ShouldGenerateWidget(SupportsTypeRow())
 									.DefaultLabel(LOCTEXT("TypeColumnHeaderName", "Type"))
 									+ SHeaderRow::Column(TransactionDetailsUI::ChangeFlagsLabel)
 									.FillWidth(15.0f)
@@ -347,11 +349,10 @@ void SUndoHistoryDetails::SetSelectedTransaction(const FTransactionDiff& InTrans
 
 	for (const auto& ObjectIt : InTransactionDiff.DiffMap)
 	{
-		UClass* ObjectClass = LoadObject<UClass>(nullptr, *ObjectIt.Value->GetOriginalObjectClassPathName().ToString());
-
-		if (ObjectClass)
+		const FSoftClassPath ClassPath = ObjectIt.Value->GetOriginalObjectClassPathName().ToString();
+		if (ReflectionData->HasClassDisplayName(ClassPath))
 		{
-			ChangedObjects.Emplace(CreateTreeNode(ObjectIt.Value->GetOriginalObjectName().ToString(), ObjectClass, ObjectIt.Value));
+			ChangedObjects.Emplace(CreateTreeNode(ObjectIt.Value->GetOriginalObjectName().ToString(), ClassPath, ObjectIt.Value));
 		}
 	}
 
@@ -364,11 +365,14 @@ void SUndoHistoryDetails::Reset()
 	ChangedObjectsTreeView->RequestTreeRefresh();
 }
 
-SUndoHistoryDetails::FUndoDetailsTreeNodePtr SUndoHistoryDetails::CreateTreeNode(const FString& InObjectName, const UClass* InObjectClass, const TSharedPtr<FTransactionObjectEvent>& InEvent) const
+SUndoHistoryDetails::FUndoDetailsTreeNodePtr SUndoHistoryDetails::CreateTreeNode(const FString& InObjectName, const FSoftClassPath& InObjectClass, const TSharedPtr<FTransactionObjectEvent>& InEvent) const
 {
-	FUndoDetailsTreeNodePtr ObjectNode = MakeShared<FUndoDetailsTreeNode>(InObjectName, InObjectClass->GetName(), FText::FromName(InEvent->GetOriginalObjectPathName()), InEvent);
+	const TOptional<FString> ClassName = ReflectionData->GetClassDisplayName(InObjectClass);
+	check(ClassName);
+	
+	FUndoDetailsTreeNodePtr ObjectNode = MakeShared<FUndoDetailsTreeNode>(InObjectName, *ClassName, FText::FromName(InEvent->GetOriginalObjectPathName()), InEvent);
 
-	TArray<FUndoHistoryUtils::FBasicPropertyInfo> BasicPropertyInfo = FUndoHistoryUtils::GetChangedPropertiesInfo(InObjectClass, InEvent->GetChangedProperties());
+	TArray<FUndoHistoryUtils::FBasicPropertyInfo> BasicPropertyInfo = FUndoHistoryUtils::GetChangedPropertiesInfo(*ReflectionData, InObjectClass, InEvent->GetChangedProperties());
 
 	for (const auto& It : BasicPropertyInfo)
 	{
@@ -378,9 +382,14 @@ SUndoHistoryDetails::FUndoDetailsTreeNodePtr SUndoHistoryDetails::CreateTreeNode
 	return ObjectNode;
 }
 
-FText SUndoHistoryDetails::CreateToolTipText(EPropertyFlags InFlags) const
+TOptional<FText> SUndoHistoryDetails::CreateToolTipText(TOptional<EPropertyFlags> InFlags) const
 {
-	TArray<const TCHAR*> Flags = ParsePropertyFlags(InFlags);
+	if (!InFlags)
+	{
+		return {};
+	}
+	
+	TArray<const TCHAR*> Flags = ParsePropertyFlags(*InFlags);
 	FString ToolTipString = "Flags:\n";
 
 	for (int32 Index = 0; Index < Flags.Num(); Index++)
@@ -396,6 +405,11 @@ FText SUndoHistoryDetails::CreateToolTipText(EPropertyFlags InFlags) const
 	return FText::FromString(ToolTipString);
 }
 
+bool SUndoHistoryDetails::SupportsTypeRow() const
+{
+	return ReflectionData->SupportsGetPropertyReflectionData();
+}
+
 void SUndoHistoryDetails::OnFilterTextChanged(const FText& InFilterText)
 {
 	SearchBoxFilter->SetRawFilterText(InFilterText);
@@ -459,7 +473,7 @@ TSharedRef<ITableRow> SUndoHistoryDetails::HandleGenerateRow(FUndoDetailsTreeNod
 		.Name(InNode->Name)
 		.Type(InNode->Type)
 		.FilterText(this, &SUndoHistoryDetails::HandleGetFilterHighlightText)
-		.ToolTipText(InNode->ToolTip);
+		.ToolTipText(InNode->ToolTip.Get(FText::GetEmpty()));
 }
 
 FText SUndoHistoryDetails::HandleGetFilterHighlightText() const
diff --git a/Engine/Source/Developer/UndoHistory/Private/Widgets/UndoHistoryUtils.h b/Engine/Source/Developer/UndoHistory/Private/Widgets/UndoHistoryUtils.h
--- a/Engine/Source/Developer/UndoHistory/Private/Widgets/UndoHistoryUtils.h
+++ b/Engine/Source/Developer/UndoHistory/Private/Widgets/UndoHistoryUtils.h
@@ -3,7 +3,6 @@
 #pragma once
 
 #include "CoreMinimal.h"
-#include "Widgets/DeclarativeSyntaxSupport.h"
 
 class FUndoHistoryUtils
 {
@@ -12,49 +11,70 @@ public:
 	struct FBasicPropertyInfo
 	{
 		FString PropertyName;
-		FString PropertyType;
-		EPropertyFlags PropertyFlags;
+		TOptional<FString> PropertyType;
+		TOptional<EPropertyFlags> PropertyFlags;
 
-		FBasicPropertyInfo(FString InPropertyName, FString InPropertyType, EPropertyFlags InPropertyFlags)
+		FBasicPropertyInfo(FString InPropertyName, TOptional<FString> InPropertyType, TOptional<EPropertyFlags> InPropertyFlags)
 			: PropertyName(MoveTemp(InPropertyName))
 			, PropertyType(MoveTemp(InPropertyType))
 			, PropertyFlags(InPropertyFlags)
 		{}
 	};
 
-	static TArray<FBasicPropertyInfo> GetChangedPropertiesInfo(const UClass* InObjectClass, const TArray<FName>& InChangedProperties)
+	static TArray<FBasicPropertyInfo> GetChangedPropertiesInfo(const UE::UndoHistory::IReflectionDataProvider& ReflectionData, const FSoftClassPath& InObjectClass, const TArray<FName>& InChangedProperties)
 	{
-		TArray<FBasicPropertyInfo> Properties;
-		FString ClassName;
-
-		if (!InObjectClass)
+		if (ReflectionData.SupportsGetPropertyReflectionData())
+		{
+			return GetChangedPropertiesInfoUsingReflectionInterface(ReflectionData, InObjectClass, InChangedProperties);
+		}
+		else
 		{
-			return Properties;
+			return GetChangedPropertiesWithoutReflectionInfo(InChangedProperties);
 		}
+	}
 
-		for (TFieldIterator<FProperty> Property(InObjectClass); Property; ++Property)
+private:
+	
+	static TArray<FBasicPropertyInfo> GetChangedPropertiesInfoUsingReflectionInterface(const UE::UndoHistory::IReflectionDataProvider& ReflectionData, const FSoftClassPath& InObjectClass, const TArray<FName>& InChangedProperties)
+	{
+		using namespace UE::UndoHistory;
+		TArray<FBasicPropertyInfo> Properties;
+		
+		for (const FName& PropertyName : InChangedProperties)
 		{
-			if (!InChangedProperties.Contains(Property->GetFName()))
-			{
-				continue;
-			}
+			const TOptional<FPropertyReflectionData> PropertyData = ReflectionData.GetPropertyReflectionData(InObjectClass, PropertyName);
+			check(PropertyData);
 
-			if (CastField<const FObjectProperty>(*Property) || Property->GetClass() == FStructProperty::StaticClass() || Property->GetClass() == FEnumProperty::StaticClass())
+			FString ClassName;
+			if (PropertyData->PropertyType == UE::UndoHistory::EPropertyType::ObjectProperty
+				|| PropertyData->PropertyType == UE::UndoHistory::EPropertyType::StructProperty 
+				|| PropertyData->PropertyType == UE::UndoHistory::EPropertyType::EnumProperty) 
 			{
-				Property->GetCPPMacroType(ClassName);
+				ClassName = PropertyData->CppMacroType;
 			}
-			else if (Property->GetClass() == FArrayProperty::StaticClass())
+			else if (PropertyData->PropertyType == UE::UndoHistory::EPropertyType::ArrayProperty)
 			{
-				Property->GetCPPMacroType(ClassName);
-				ClassName = FString::Printf(TEXT("TArray<%s>"), *ClassName);
+				ClassName = PropertyData->CppMacroType;
+				ClassName = FString::Printf(TEXT("TArray<%s>"), *ClassName); 
 			}
 			else
 			{
-				ClassName = Property->GetClass()->GetName();
+				ClassName = PropertyData->TypeName;
 				ClassName.RemoveFromEnd("Property");
 			}
 
-			Properties.Emplace(Property->GetName(), MoveTemp(ClassName), Property->GetPropertyFlags());
+			Properties.Emplace(PropertyName.ToString(), MoveTemp(ClassName), PropertyData->PropertyFlags);
+		}
+		
+		return Properties;
+	}
+
+	static TArray<FBasicPropertyInfo> GetChangedPropertiesWithoutReflectionInfo(const TArray<FName>& InChangedProperties)
+	{
+		TArray<FBasicPropertyInfo> Properties;
+		for (const FName& PropertyName : InChangedProperties)
+		{
+			Properties.Emplace(PropertyName.ToString(), TOptional<FString>{}, TOptional<EPropertyFlags>{});
 		}
 		return Properties;
 	}
diff --git a/Engine/Source/Developer/UndoHistory/Public/Widgets/IReflectionDataProvider.h b/Engine/Source/Developer/UndoHistory/Public/Widgets/IReflectionDataProvider.h
new file mode 100644
--- /dev/null
+++ b/Engine/Source/Developer/UndoHistory/Public/Widgets/IReflectionDataProvider.h
@@ -0,0 +1,53 @@
+// Copyright Epic Games, Inc. All Rights Reserved.
+
+#pragma once
+
+#include "CoreMinimal.h"
+#include "UObject/UnrealType.h"
+
+namespace UE::UndoHistory
+{
+	enum class EPropertyType : uint8
+	{
+		ObjectProperty,
+		StructProperty,
+		EnumProperty,
+		ArrayProperty,
+		Other
+	};
+
+	struct FPropertyReflectionData
+	{
+		/** The result of FProperty::GetCPPMacroType (the part that is returned by ref, e.g. "FName"). */
+		FString CppMacroType;
+		/** The result of FProperty::GetClass()::GetName() */
+		FString TypeName;
+		EPropertyFlags PropertyFlags;
+		EPropertyType PropertyType;
+	};
+
+	
+	/**
+	 * Provides reflection data to SUndoHistoryDetails.
+	 * 
+	 * On editor builds reflection data is fully available however on non-editor builds there is no straightforward way of obtaining it.
+	 * For example, Concert implements this interface to return data where available.
+	 */
+	class UNDOHISTORY_API IReflectionDataProvider
+	{
+	public:
+
+		virtual bool HasClassDisplayName(const FSoftClassPath& ClassPath) const = 0;
+		/** @return What UClass::GetName() would return, if available */
+		virtual TOptional<FString> GetClassDisplayName(const FSoftClassPath& ClassPath) const = 0;
+
+		virtual bool SupportsGetPropertyReflectionData() const = 0;
+		/** @return Various misc data required for UI, e.g. for the type column */		
+		virtual TOptional<FPropertyReflectionData> GetPropertyReflectionData(const FSoftClassPath& ClassPath, FName PropertyName) const = 0;
+		
+		virtual ~IReflectionDataProvider() = default;
+	};
+
+	/** Gets the default provider that just calls LoadObject. Works on editor builds. Will not work in most programs, like MultiUserServer. */
+	UNDOHISTORY_API TSharedRef<IReflectionDataProvider> CreateDefaultReflectionProvider();
+}
diff --git a/Engine/Source/Developer/UndoHistory/Public/Widgets/SUndoHistoryDetails.h b/Engine/Source/Developer/UndoHistory/Public/Widgets/SUndoHistoryDetails.h
--- a/Engine/Source/Developer/UndoHistory/Public/Widgets/SUndoHistoryDetails.h
+++ b/Engine/Source/Developer/UndoHistory/Public/Widgets/SUndoHistoryDetails.h
@@ -3,6 +3,7 @@
 #pragma once
 
 #include "CoreMinimal.h"
+#include "IReflectionDataProvider.h"
 #include "Widgets/DeclarativeSyntaxSupport.h"
 #include "Widgets/SCompoundWidget.h"
 #include "Widgets/Views/STreeView.h"
@@ -12,7 +13,7 @@ class UNDOHISTORY_API SUndoHistoryDetails : public SCompoundWidget
 {
 
 public:
-	SLATE_BEGIN_ARGS(SUndoHistoryDetails) { }
+	SLATE_BEGIN_ARGS(SUndoHistoryDetails) {}
 	SLATE_END_ARGS()
 
 public:
@@ -22,7 +23,7 @@ public:
 	 *
 	 * @param InArgs The declaration data for this widget.
 	 */
-	void Construct(const FArguments& InArgs);
+	void Construct(const FArguments& InArgs, TSharedRef<UE::UndoHistory::IReflectionDataProvider> InReflectionData = UE::UndoHistory::CreateDefaultReflectionProvider());
 
 	/**
 	 * Set the transaction to be displayed in the details panel.
@@ -48,12 +49,12 @@ private:
 	struct FUndoDetailsTreeNode
 	{
 		FString Name;
-		FString Type;
-		FText ToolTip;
+		TOptional<FString> Type;
+		TOptional<FText> ToolTip;
 		TSharedPtr<FTransactionObjectEvent> TransactionEvent;
 		TArray<FUndoDetailsTreeNodePtr> Children;
 
-		FUndoDetailsTreeNode(FString InName, FString InType, FText InToolTip, TSharedPtr<FTransactionObjectEvent> InTransactionEvent = nullptr)
+		FUndoDetailsTreeNode(FString InName, TOptional<FString> InType, TOptional<FText> InToolTip, TSharedPtr<FTransactionObjectEvent> InTransactionEvent = nullptr)
 			: Name(MoveTemp(InName))
 			, Type(MoveTemp(InType))
 			, ToolTip(MoveTemp(InToolTip))
@@ -64,10 +65,13 @@ private:
 private:
 
 	/** Create a changed object node. */
-	FUndoDetailsTreeNodePtr CreateTreeNode(const FString& InObjectName, const UClass* InObjectClass, const TSharedPtr<FTransactionObjectEvent>& InEvent) const;
+	FUndoDetailsTreeNodePtr CreateTreeNode(const FString& InObjectName, const FSoftClassPath& InObjectClass, const TSharedPtr<FTransactionObjectEvent>& InEvent) const;
 
 	/** Create a tooltip text for a property. */
-	FText CreateToolTipText(EPropertyFlags InFlags) const;
+	TOptional<FText> CreateToolTipText(TOptional<EPropertyFlags> InFlags) const;
+
+	/** Whether the type row will be generated */
+	bool SupportsTypeRow() const;
 
 	/** Callback to handle a change in the filter box. */
 	void OnFilterTextChanged(const FText& InFilterText);
@@ -101,6 +105,9 @@ private:
 
 private:
 
+	/** Gives us information about property data (gives programs the opportunity to inject custom handler since not all reflection data is available in programs). */
+	TSharedPtr<UE::UndoHistory::IReflectionDataProvider> ReflectionData;
+	
 	/** Holds the ChangedObjects TreeView. */
 	TSharedPtr<STreeView<FUndoDetailsTreeNodePtr> > ChangedObjectsTreeView;
 
