Fix undo history crash when temporary property not yet in object reflection data.

diff --git a/Engine/Source/Developer/UndoHistory/Private/Widgets/UndoHistoryUtils.h b/Engine/Source/Developer/UndoHistory/Private/Widgets/UndoHistoryUtils.h
--- a/Engine/Source/Developer/UndoHistory/Private/Widgets/UndoHistoryUtils.h
+++ b/Engine/Source/Developer/UndoHistory/Private/Widgets/UndoHistoryUtils.h
@@ -45,28 +45,32 @@ private:
 		
 		for (const FName& PropertyName : InChangedProperties)
 		{
-			const TOptional<FPropertyReflectionData> PropertyData = ReflectionData.GetPropertyReflectionData(InObjectClass, PropertyName);
-			check(PropertyData);
-
-			FString ClassName;
-			if (PropertyData->PropertyType == UE::UndoHistory::EPropertyType::ObjectProperty
-				|| PropertyData->PropertyType == UE::UndoHistory::EPropertyType::StructProperty 
-				|| PropertyData->PropertyType == UE::UndoHistory::EPropertyType::EnumProperty) 
-			{
-				ClassName = PropertyData->CppMacroType;
-			}
-			else if (PropertyData->PropertyType == UE::UndoHistory::EPropertyType::ArrayProperty)
+			if (TOptional<FPropertyReflectionData> PropertyData = ReflectionData.GetPropertyReflectionData(InObjectClass, PropertyName))
 			{
-				ClassName = PropertyData->CppMacroType;
-				ClassName = FString::Printf(TEXT("TArray<%s>"), *ClassName); 
+				FString ClassName;
+				if (PropertyData->PropertyType == UE::UndoHistory::EPropertyType::ObjectProperty
+					|| PropertyData->PropertyType == UE::UndoHistory::EPropertyType::StructProperty
+					|| PropertyData->PropertyType == UE::UndoHistory::EPropertyType::EnumProperty)
+				{
+					ClassName = PropertyData->CppMacroType;
+				}
+				else if (PropertyData->PropertyType == UE::UndoHistory::EPropertyType::ArrayProperty)
+				{
+					ClassName = PropertyData->CppMacroType;
+					ClassName = FString::Printf(TEXT("TArray<%s>"), *ClassName);
+				}
+				else
+				{
+					ClassName = PropertyData->TypeName;
+					ClassName.RemoveFromEnd("Property");
+				}
+
+				Properties.Emplace(PropertyName.ToString(), MoveTemp(ClassName), PropertyData->PropertyFlags);
 			}
 			else
 			{
-				ClassName = PropertyData->TypeName;
-				ClassName.RemoveFromEnd("Property");
+				Properties.Emplace(PropertyName.ToString(), TOptional<FString>{}, TOptional<EPropertyFlags>{});
 			}
-
-			Properties.Emplace(PropertyName.ToString(), MoveTemp(ClassName), PropertyData->PropertyFlags);
 		}
 		
 		return Properties;
