Updated the undo history to actually show the historic diff info

It now shows the diff that was sent to an object's PostTransacted function the last time the transaction was applied (either from finalization, or undo/redo), rather than re-generating the diff against the current state of the object (which was inaccurate if there were multiple edits made to the object in the history). The object names shown in the tree are now the post-change name/outer, and the state change tooltip has been updated to also show the pre-change name/outer; this makes it easier to track name/outer changes over time.

diff --git a/Engine/Source/Developer/UndoHistory/Private/Widgets/SUndoHistoryDetails.cpp b/Engine/Source/Developer/UndoHistory/Private/Widgets/SUndoHistoryDetails.cpp
--- a/Engine/Source/Developer/UndoHistory/Private/Widgets/SUndoHistoryDetails.cpp
+++ b/Engine/Source/Developer/UndoHistory/Private/Widgets/SUndoHistoryDetails.cpp
@@ -17,6 +17,8 @@
 #include "Widgets/Layout/SSeparator.h"
 #include "HAL/PlatformApplicationMisc.h"
 #include "UndoHistoryUtils.h"
+#include "Misc/AsciiSet.h"
+#include "Misc/PackageName.h"
 #include "Misc/TransactionObjectEvent.h"
 #include "Algo/Transform.h"
 
@@ -29,45 +31,51 @@ namespace TransactionDetailsUI
 	static const FName TypeLabel(TEXT("Type"));
 	static const FName ChangeFlagsLabel(TEXT("ChangeFlags"));
  
-	TSharedRef<SImage> CreateRenameIcon()
+	TSharedRef<SWidget> CreateIconSpacer()
+	{
+		return SNew(SBox)
+			.WidthOverride(16)
+			.HeightOverride(16);
+	}
+
+	TSharedRef<SWidget> CreateRenameIcon(const FName OriginalName)
 	{
 		return SNew(SImage)
-			.Image(FAppStyle::Get().GetBrush("Icons.Rename"))
-			.ToolTipText(LOCTEXT("RenameChangeToolTip", "Renamed."))
-			.ColorAndOpacity(FSlateColor::UseForeground()); 
+			.Image(FAppStyle::Get().GetBrush("Icons.Edit"))
+			.ToolTipText(FText::Format(LOCTEXT("RenameChangeToolTip", "Renamed.\nOriginal Name: {0}"), FText::FromName(OriginalName)))
+			.ColorAndOpacity(FSlateColor::UseForeground());
 	}
 
-	TSharedRef<SImage> CreateOuterChangeIcon()
+	TSharedRef<SWidget> CreateOuterChangeIcon(const FName OriginalOuter)
 	{
 		return SNew(SImage)
 			.Image(FAppStyle::Get().GetBrush("Icons.CircleArrowUp"))
- 
-			.ToolTipText(LOCTEXT("OuterChangeToolTip", "Has an outer change."))
-			.ColorAndOpacity(FSlateColor::UseForeground()); 
+			.ToolTipText(FText::Format(LOCTEXT("OuterChangeToolTip", "Has an outer change.\nOriginal Outer: {0}"), FText::FromName(OriginalOuter)))
+			.ColorAndOpacity(FSlateColor::UseForeground());
 	}
-	TSharedRef<SImage> CreatePendingKillIcon()
+
+	TSharedRef<SWidget> CreatePendingKillIcon()
 	{
 		return SNew(SImage)
 			.Image(FAppStyle::Get().GetBrush("Icons.FilledCircle"))
- 
 			.ToolTipText(LOCTEXT("PendingKillToolTip", "Has a pending kill change."))
 			.ColorAndOpacity(FSlateColor::UseForeground());
 	}
-	TSharedRef<SImage> CreateNonPropertyChangeIcon()
+
+	TSharedRef<SWidget> CreateNonPropertyChangeIcon()
 	{
 		return SNew(SImage)
 			.Image(FAppStyle::Get().GetBrush("Icons.Settings"))
- 
 			.ToolTipText(LOCTEXT("NonPropertyChangeToolTip", "Has a non-property change."))
 			.ColorAndOpacity(FSlateColor::UseForeground());
 	}
-	TSharedRef<SImage> CreateAnnotationIcon()
+
+	TSharedRef<SWidget> CreateAnnotationIcon()
 	{
 		return SNew(SImage)
 			.Image(FAppStyle::Get().GetBrush("Icons.Info"))
- 
 			.ToolTipText(LOCTEXT("AnnotationToolTip", "Has an annotation."))
-			.ColorAndOpacity(FSlateColor::UseForeground()); 
+			.ColorAndOpacity(FSlateColor::UseForeground());
 	}
 };
 
@@ -85,7 +93,7 @@ public:
 	{ }
 		SLATE_ARGUMENT(TSharedPtr<FTransactionObjectEvent>, TransactionEvent)
 		SLATE_ARGUMENT(FString, Name)
-		SLATE_ARGUMENT(TOptional<FString>, Type)
+		SLATE_ARGUMENT(FString, Type)
 		SLATE_ATTRIBUTE(FText, FilterText)
 	SLATE_END_ARGS()
 
@@ -131,10 +139,10 @@ public:
 							.HighlightText(FilterText)
 					];
 		}
-		else if (ColumnName == TransactionDetailsUI::TypeLabel && ensure(Type))
+		else if (ColumnName == TransactionDetailsUI::TypeLabel)
 		{
 			return SNew(STextBlock)
-				.Text(FText::FromString(*Type));
+				.Text(FText::FromString(Type));
 		}
 		else if (ColumnName == TransactionDetailsUI::ChangeFlagsLabel)
 		{
@@ -161,50 +169,35 @@ private:
 	{
 		TSharedRef<SHorizontalBox> HorizontalBox = SNew(SHorizontalBox);
 
-		if (TransactionEvent->HasNameChange())
-		{
-			HorizontalBox->AddSlot()
-				.Padding(FMargin(2.0f, 0.0f))
-				[
-					TransactionDetailsUI::CreateRenameIcon()
-				];
-		}
-
-		if (TransactionEvent->HasOuterChange())
-		{
-			HorizontalBox->AddSlot()
-				.Padding(FMargin(2.0f, 0.0f))
-				[
-					TransactionDetailsUI::CreateOuterChangeIcon()
-				];
-		}
-
-		if (TransactionEvent->HasPendingKillChange())
-		{
-			HorizontalBox->AddSlot()
-				.Padding(FMargin(2.0f, 0.0f))
-				[
-					TransactionDetailsUI::CreatePendingKillIcon()
-				];
-		}
-
-		if (TransactionEvent->HasNonPropertyChanges())
-		{
-			HorizontalBox->AddSlot()
-				.Padding(FMargin(2.0f, 0.0f))
-				[
-					TransactionDetailsUI::CreateNonPropertyChangeIcon()
-				];
-		}
-
-		if (TransactionEvent->GetAnnotation().IsValid())
-		{
-			HorizontalBox->AddSlot()
-				.Padding(FMargin(2.0f, 0.0f))
-				[
-					TransactionDetailsUI::CreateAnnotationIcon()
-				];
-		}
+		HorizontalBox->AddSlot()
+			.Padding(FMargin(2.0f, 0.0f))
+			[
+				TransactionEvent->HasNameChange() ? TransactionDetailsUI::CreateRenameIcon(TransactionEvent->GetOriginalObjectName()) : TransactionDetailsUI::CreateIconSpacer()
+			];
+
+		HorizontalBox->AddSlot()
+			.Padding(FMargin(2.0f, 0.0f))
+			[
+				TransactionEvent->HasOuterChange() ? TransactionDetailsUI::CreateOuterChangeIcon(TransactionEvent->GetOriginalObjectOuterPathName()) : TransactionDetailsUI::CreateIconSpacer()
+			];
+
+		HorizontalBox->AddSlot()
+			.Padding(FMargin(2.0f, 0.0f))
+			[
+				TransactionEvent->HasPendingKillChange() ? TransactionDetailsUI::CreatePendingKillIcon() : TransactionDetailsUI::CreateIconSpacer()
+			];
+
+		HorizontalBox->AddSlot()
+			.Padding(FMargin(2.0f, 0.0f))
+			[
+				TransactionEvent->HasNonPropertyChanges() ? TransactionDetailsUI::CreateNonPropertyChangeIcon() : TransactionDetailsUI::CreateIconSpacer()
+			];
+
+		HorizontalBox->AddSlot()
+			.Padding(FMargin(2.0f, 0.0f))
+			[
+				TransactionEvent->GetAnnotation().IsValid() ? TransactionDetailsUI::CreateAnnotationIcon() : TransactionDetailsUI::CreateIconSpacer()
+			];
 
 		return HorizontalBox;
 	}
@@ -213,7 +206,7 @@ private:
 
 	TSharedPtr<FTransactionObjectEvent> TransactionEvent;
 	FString Name;
-	TOptional<FString> Type;
+	FString Type;
 	TAttribute<FText> FilterText;
 
 };
@@ -352,10 +345,16 @@ void SUndoHistoryDetails::SetSelectedTransaction(const FTransactionDiff& InTrans
 		const FSoftClassPath ClassPath = ObjectIt.Value->GetOriginalObjectClassPathName().ToString();
 		if (ReflectionData->HasClassDisplayName(ClassPath))
 		{
-			ChangedObjects.Emplace(CreateTreeNode(ObjectIt.Value->GetOriginalObjectName().ToString(), ClassPath, ObjectIt.Value));
+			ChangedObjects.Emplace(CreateTreeNode(ObjectIt.Key.ToString(), ClassPath, ObjectIt.Value));
 		}
 	}
 
+	// Sort the list so that outer objects appear before inner objects
+	ChangedObjects.StableSort([](const FUndoDetailsTreeNodePtr& A, const FUndoDetailsTreeNodePtr& B)
+	{
+		return A->PathDepth < B->PathDepth;
+	});
+
 	Populate();
 }
 
@@ -365,12 +364,21 @@ void SUndoHistoryDetails::Reset()
 	ChangedObjectsTreeView->RequestTreeRefresh();
 }
 
-SUndoHistoryDetails::FUndoDetailsTreeNodePtr SUndoHistoryDetails::CreateTreeNode(const FString& InObjectName, const FSoftClassPath& InObjectClass, const TSharedPtr<FTransactionObjectEvent>& InEvent) const
+SUndoHistoryDetails::FUndoDetailsTreeNodePtr SUndoHistoryDetails::CreateTreeNode(const FString& InObjectPathName, const FSoftClassPath& InObjectClass, const TSharedPtr<FTransactionObjectEvent>& InEvent) const
 {
 	const TOptional<FString> ClassName = ReflectionData->GetClassDisplayName(InObjectClass);
 	check(ClassName);
 	
-	FUndoDetailsTreeNodePtr ObjectNode = MakeShared<FUndoDetailsTreeNode>(InObjectName, *ClassName, FText::FromName(InEvent->GetOriginalObjectPathName()), InEvent);
+	int32 ObjectPathDepth = 0;
+	{
+		constexpr FAsciiSet PathSepratorCharacters(".:");
+		for (const TCHAR* It = *InObjectPathName; *It != 0; It = FAsciiSet::FindFirstOrEnd(It + 1, PathSepratorCharacters))
+		{
+			++ObjectPathDepth;
+		}
+	}
+
+	FUndoDetailsTreeNodePtr ObjectNode = MakeShared<FUndoDetailsTreeNode>(FPackageName::ObjectPathToObjectName(InObjectPathName), *ClassName, FText::AsCultureInvariant(InObjectPathName), InEvent, ObjectPathDepth);
 
 	TArray<FUndoHistoryUtils::FBasicPropertyInfo> BasicPropertyInfo = FUndoHistoryUtils::GetChangedPropertiesInfo(*ReflectionData, InObjectClass, InEvent->GetChangedProperties());
 
@@ -382,27 +390,17 @@ SUndoHistoryDetails::FUndoDetailsTreeNodePtr SUndoHistoryDetails::CreateTreeNode
 	return ObjectNode;
 }
 
-TOptional<FText> SUndoHistoryDetails::CreateToolTipText(TOptional<EPropertyFlags> InFlags) const
+FText SUndoHistoryDetails::CreateToolTipText(EPropertyFlags InFlags) const
 {
-	if (!InFlags)
+	if (InFlags == CPF_None)
 	{
 		return {};
 	}
 	
-	TArray<const TCHAR*> Flags = ParsePropertyFlags(*InFlags);
-	FString ToolTipString = "Flags:\n";
-
-	for (int32 Index = 0; Index < Flags.Num(); Index++)
-	{
-		ToolTipString += FString(Flags[Index]);
+	FString ToolTipString = TEXT("Flags:\n");
+	ToolTipString += FString::Join(ParsePropertyFlags(InFlags), TEXT("\n"));
 
-		if (Index != Flags.Num() - 1)
-		{
-			ToolTipString += "\n";
-		}
-	}
-
-	return FText::FromString(ToolTipString);
+	return FText::FromString(MoveTemp(ToolTipString));
 }
 
 bool SUndoHistoryDetails::SupportsTypeRow() const
@@ -459,21 +457,12 @@ void SUndoHistoryDetails::Populate()
 
 TSharedRef<ITableRow> SUndoHistoryDetails::HandleGenerateRow(FUndoDetailsTreeNodePtr InNode, const TSharedRef<STableViewBase>& OwnerTable) const
 {
-	if (InNode->TransactionEvent.IsValid())
-	{
-		return SNew(SUndoHistoryDetailsRow, OwnerTable)
-			.Name(InNode->Name)
-			.Type(InNode->Type)
-			.ToolTipText(FText::FromName(InNode->TransactionEvent->GetOriginalObjectPathName()))
-			.FilterText(this, &SUndoHistoryDetails::HandleGetFilterHighlightText)
-			.TransactionEvent(InNode->TransactionEvent);
-	}
-
 	return SNew(SUndoHistoryDetailsRow, OwnerTable)
 		.Name(InNode->Name)
 		.Type(InNode->Type)
 		.FilterText(this, &SUndoHistoryDetails::HandleGetFilterHighlightText)
-		.ToolTipText(InNode->ToolTip.Get(FText::GetEmpty()));
+		.ToolTipText(InNode->ToolTip)
+		.TransactionEvent(InNode->TransactionEvent);
 }
 
 FText SUndoHistoryDetails::HandleGetFilterHighlightText() const
diff --git a/Engine/Source/Developer/UndoHistory/Private/Widgets/UndoHistoryUtils.h b/Engine/Source/Developer/UndoHistory/Private/Widgets/UndoHistoryUtils.h
--- a/Engine/Source/Developer/UndoHistory/Private/Widgets/UndoHistoryUtils.h
+++ b/Engine/Source/Developer/UndoHistory/Private/Widgets/UndoHistoryUtils.h
@@ -14,10 +14,10 @@ public:
 	struct FBasicPropertyInfo
 	{
 		FString PropertyName;
-		TOptional<FString> PropertyType;
-		TOptional<EPropertyFlags> PropertyFlags;
+		FString PropertyType;
+		EPropertyFlags PropertyFlags = CPF_None;
 
-		FBasicPropertyInfo(FString InPropertyName, TOptional<FString> InPropertyType, TOptional<EPropertyFlags> InPropertyFlags)
+		FBasicPropertyInfo(FString InPropertyName, FString InPropertyType = FString(), EPropertyFlags InPropertyFlags = CPF_None)
 			: PropertyName(MoveTemp(InPropertyName))
 			, PropertyType(MoveTemp(InPropertyType))
 			, PropertyFlags(InPropertyFlags)
@@ -69,7 +69,7 @@ private:
 			}
 			else
 			{
-				Properties.Emplace(PropertyName.ToString(), TOptional<FString>{}, TOptional<EPropertyFlags>{});
+				Properties.Emplace(PropertyName.ToString());
 			}
 		}
 		
@@ -81,7 +81,7 @@ private:
 		TArray<FBasicPropertyInfo> Properties;
 		for (const FName& PropertyName : InChangedProperties)
 		{
-			Properties.Emplace(PropertyName.ToString(), TOptional<FString>{}, TOptional<EPropertyFlags>{});
+			Properties.Emplace(PropertyName.ToString());
 		}
 		return Properties;
 	}
diff --git a/Engine/Source/Developer/UndoHistory/Public/Widgets/SUndoHistoryDetails.h b/Engine/Source/Developer/UndoHistory/Public/Widgets/SUndoHistoryDetails.h
--- a/Engine/Source/Developer/UndoHistory/Public/Widgets/SUndoHistoryDetails.h
+++ b/Engine/Source/Developer/UndoHistory/Public/Widgets/SUndoHistoryDetails.h
@@ -49,26 +49,28 @@ private:
 	struct FUndoDetailsTreeNode
 	{
 		FString Name;
-		TOptional<FString> Type;
-		TOptional<FText> ToolTip;
+		FString Type;
+		FText ToolTip;
 		TSharedPtr<FTransactionObjectEvent> TransactionEvent;
 		TArray<FUndoDetailsTreeNodePtr> Children;
+		int32 PathDepth = 0;
 
-		FUndoDetailsTreeNode(FString InName, TOptional<FString> InType, TOptional<FText> InToolTip, TSharedPtr<FTransactionObjectEvent> InTransactionEvent = nullptr)
+		FUndoDetailsTreeNode(FString InName, FString InType, FText InToolTip, TSharedPtr<FTransactionObjectEvent> InTransactionEvent = nullptr, int32 InPathDepth = 0)
 			: Name(MoveTemp(InName))
 			, Type(MoveTemp(InType))
 			, ToolTip(MoveTemp(InToolTip))
 			, TransactionEvent(MoveTemp(InTransactionEvent))
+			, PathDepth(InPathDepth)
 		{}
 	};
 
 private:
 
 	/** Create a changed object node. */
-	FUndoDetailsTreeNodePtr CreateTreeNode(const FString& InObjectName, const FSoftClassPath& InObjectClass, const TSharedPtr<FTransactionObjectEvent>& InEvent) const;
+	FUndoDetailsTreeNodePtr CreateTreeNode(const FString& InObjectPathName, const FSoftClassPath& InObjectClass, const TSharedPtr<FTransactionObjectEvent>& InEvent) const;
 
 	/** Create a tooltip text for a property. */
-	TOptional<FText> CreateToolTipText(TOptional<EPropertyFlags> InFlags) const;
+	FText CreateToolTipText(EPropertyFlags InFlags) const;
 
 	/** Whether the type row will be generated */
 	bool SupportsTypeRow() const;
diff --git a/Engine/Source/Editor/UndoHistoryEditor/Private/Widgets/SUndoHistory.cpp b/Engine/Source/Editor/UndoHistoryEditor/Private/Widgets/SUndoHistory.cpp
--- a/Engine/Source/Editor/UndoHistoryEditor/Private/Widgets/SUndoHistory.cpp
+++ b/Engine/Source/Editor/UndoHistoryEditor/Private/Widgets/SUndoHistory.cpp
@@ -298,7 +298,7 @@ void SUndoHistory::SelectLastTransaction()
 			if (InnerTransaction)
 			{
 				UndoListView->SetSelection(TransactionInfo);
-				UndoDetailsView->SetSelectedTransaction(InnerTransaction->GenerateDiff());
+				UndoDetailsView->SetSelectedTransaction(InnerTransaction->GetHistoricDiff());
 			}
 		}
 		else
@@ -406,7 +406,7 @@ void SUndoHistory::HandleUndoListSelectionChanged(TSharedPtr<FTransactionInfo> I
 
 	if (const FTransaction* InnerTransaction = InItem->GetTransaction())
 	{
-		UndoDetailsView->SetSelectedTransaction(InnerTransaction->GenerateDiff());
+		UndoDetailsView->SetSelectedTransaction(InnerTransaction->GetHistoricDiff());
 	}
 
 	if (SelectInfo == ESelectInfo::OnMouseClick || SelectInfo == ESelectInfo::Direct)
diff --git a/Engine/Source/Editor/UnrealEd/Classes/Editor/Transactor.h b/Engine/Source/Editor/UnrealEd/Classes/Editor/Transactor.h
--- a/Engine/Source/Editor/UnrealEd/Classes/Editor/Transactor.h
+++ b/Engine/Source/Editor/UnrealEd/Classes/Editor/Transactor.h
@@ -322,6 +322,9 @@ protected:
 	/** Objects that will be changed directly by the transaction, empty when not transacting */
 	TMap<UObject*, FChangedObjectValue> ChangedObjects;
 
+	/** Cached state of the diff the last time this transaction was applied to an object (via PostTransacted for non-snapshot transactions) to allow historic introspection of the diff regardless of the current state of the object */
+	TMap<FName, TSharedPtr<FTransactionObjectEvent>> HistoricDiffMap;
+
 public:
 	// Constructor.
 	FTransaction(  const TCHAR* InContext=nullptr, const FText& InTitle=FText(), bool bInFlip=false )
@@ -485,10 +488,18 @@ public:
 	UNREALED_API void DumpObjectMap(FOutputDevice& Ar) const;
 
 	/**
-	 * Create a map that holds information about objects of a given transaction.
+	 * Create a map that holds information about objects of a given transaction compared to the current state of the object.
 	 */
 	UNREALED_API FTransactionDiff GenerateDiff() const;
 
+	/**
+	 * Create a map that holds information about objects of a given transaction compared to their previous state the last time this transaction was applied to those objects (via PostTransacted).
+	 */
+	UNREALED_API FTransactionDiff GetHistoricDiff() const;
+
+private:
+	void PostTransactAndSaveHistoricDiff(UObject* Obj, FTransactionObjectEvent&& TransactionEvent);
+
 	// Transaction friends.
 	friend FArchive& operator<<( FArchive& Ar, FTransaction::FObjectRecord& R );
 
diff --git a/Engine/Source/Editor/UnrealEd/Private/EditorTransaction.cpp b/Engine/Source/Editor/UnrealEd/Private/EditorTransaction.cpp
--- a/Engine/Source/Editor/UnrealEd/Private/EditorTransaction.cpp
+++ b/Engine/Source/Editor/UnrealEd/Private/EditorTransaction.cpp
@@ -925,6 +925,7 @@ void FTransaction::Apply()
 	};
 
 	TArray<ULevel*> LevelsToCommitModelSurface;
+	HistoricDiffMap.Reset();
 	for (const TTuple<UObject*, FChangedObjectValue>& ChangedObjectIt : ChangedObjects)
 	{
 		UObject* ChangedObject = ChangedObjectIt.Key;
@@ -987,10 +988,10 @@ void FTransaction::Apply()
 				ChangedObjectTransactionAnnotation->ComputeAdditionalObjectChanges(InitialSerializedObject.ObjectAnnotation.Get(), AdditionalObjectChanges);
 			}
 
-			ChangedObject->PostTransacted(FTransactionObjectEvent(Id, OperationId, ETransactionObjectEventType::UndoRedo, ETransactionObjectChangeCreatedBy::TransactionRecord, FTransactionObjectChange{ InitialSerializedObject.ObjectId, InitialSerializedObject.ObjectFlags, MoveTemp(PrimaryDeltaChange) }, ChangedObjectTransactionAnnotation));
+			PostTransactAndSaveHistoricDiff(ChangedObject, FTransactionObjectEvent(Id, OperationId, ETransactionObjectEventType::UndoRedo, ETransactionObjectChangeCreatedBy::TransactionRecord, FTransactionObjectChange{ InitialSerializedObject.ObjectId, InitialSerializedObject.ObjectFlags, MoveTemp(PrimaryDeltaChange) }, ChangedObjectTransactionAnnotation));
 			for (const TTuple<UObject*, FTransactionObjectChange>& AdditionalObjectChangePair : AdditionalObjectChanges)
 			{
-				AdditionalObjectChangePair.Key->PostTransacted(FTransactionObjectEvent(Id, OperationId, ETransactionObjectEventType::UndoRedo, ETransactionObjectChangeCreatedBy::TransactionAnnotation, AdditionalObjectChangePair.Value, nullptr));
+				PostTransactAndSaveHistoricDiff(AdditionalObjectChangePair.Key, FTransactionObjectEvent(Id, OperationId, ETransactionObjectEventType::UndoRedo, ETransactionObjectChangeCreatedBy::TransactionAnnotation, AdditionalObjectChangePair.Value, nullptr));
 			}
 		}
 	}
@@ -1048,6 +1049,7 @@ void FTransaction::Finalize()
 		return Cast<UActorComponent>(&A) != nullptr;
 	});
 
+	HistoricDiffMap.Reset();
 	for (const TTuple<UObject*, FChangedObjectValue>& ChangedObjectIt : ChangedObjects)
 	{
 		const FObjectRecord& ChangedObjectRecord = Records[ChangedObjectIt.Value.RecordIndex];
@@ -1069,10 +1071,10 @@ void FTransaction::Finalize()
 				ChangedObjectTransactionAnnotation->ComputeAdditionalObjectChanges(ChangedObjectIt.Value.AnnotationSnapshot.Get(), AdditionalObjectChanges);
 			}
 
-			ChangedObject->PostTransacted(FTransactionObjectEvent(Id, OperationId, ETransactionObjectEventType::Finalized, ETransactionObjectChangeCreatedBy::TransactionRecord, FTransactionObjectChange{ InitialSerializedObject.ObjectId, InitialSerializedObject.ObjectFlags, DeltaChange }, ChangedObjectTransactionAnnotation));
+			PostTransactAndSaveHistoricDiff(ChangedObject, FTransactionObjectEvent(Id, OperationId, ETransactionObjectEventType::Finalized, ETransactionObjectChangeCreatedBy::TransactionRecord, FTransactionObjectChange{ InitialSerializedObject.ObjectId, InitialSerializedObject.ObjectFlags, DeltaChange }, ChangedObjectTransactionAnnotation));
 			for (const TTuple<UObject*, FTransactionObjectChange>& AdditionalObjectChangePair : AdditionalObjectChanges)
 			{
-				AdditionalObjectChangePair.Key->PostTransacted(FTransactionObjectEvent(Id, OperationId, ETransactionObjectEventType::Finalized, ETransactionObjectChangeCreatedBy::TransactionAnnotation, AdditionalObjectChangePair.Value, nullptr));
+				PostTransactAndSaveHistoricDiff(AdditionalObjectChangePair.Key, FTransactionObjectEvent(Id, OperationId, ETransactionObjectEventType::Finalized, ETransactionObjectChangeCreatedBy::TransactionAnnotation, AdditionalObjectChangePair.Value, nullptr));
 			}
 		}
 	}
@@ -1151,6 +1153,19 @@ FTransactionDiff FTransaction::GenerateDiff() const
 	return TransactionDiff;
 }
 
+FTransactionDiff FTransaction::GetHistoricDiff() const
+{
+	FTransactionDiff TransactionDiff{Id, Title.ToString()};
+	TransactionDiff.DiffMap = HistoricDiffMap;
+	return TransactionDiff;
+}
+
+void FTransaction::PostTransactAndSaveHistoricDiff(UObject* Obj, FTransactionObjectEvent&& TransactionEvent)
+{
+	Obj->PostTransacted(TransactionEvent);
+	HistoricDiffMap.Add(*Obj->GetPathName(), MakeShared<FTransactionObjectEvent>(MoveTemp(TransactionEvent)));
+}
+
 
 /*-----------------------------------------------------------------------------
 	Transaction tracking system.
